Generated metadata files for low level tests

diff --git a/Engine/Build/LowLevelTests/Foundation.xml b/Engine/Build/LowLevelTests/Foundation.xml
new file mode 100644
--- /dev/null
+++ b/Engine/Build/LowLevelTests/Foundation.xml
@@ -0,0 +1,39 @@
+﻿<?xml version="1.0" encoding="utf-8"?>
+<BuildGraph xmlns:xsi="http://www.w3.org/2001/XMLSchema-instance" xsi:schemaLocation="http://www.epicgames.com/BuildGraph ../../Build/Graph/Schema.xsd" xmlns="http://www.epicgames.com/BuildGraph">
+  <Property Name="TestNames" Value="$(TestNames);Foundation" />
+  <Property Name="FoundationUsesCatch2" Value="True" />
+  <Property Name="FoundationDisabled" Value="False" />
+  <Property Name="FoundationShort" Value="Foundation" />
+  <Property Name="FoundationTarget" Value="FoundationTests" />
+  <Property Name="FoundationBinariesRelative" Value="Engine\Binaries" />
+  <Property Name="FoundationReportType" Value="xml" />
+  <Property Name="FoundationInitialExtraArgs" />
+  <Property Name="FoundationHasAfterSteps" Value="False" />
+  <Property Name="FoundationUsesCatch2" Value="True" />
+  <Option Name="RunFoundationTests" DefaultValue="False" Description="Run Foundation Tests" />
+  <Property Name="FoundationSupportedPlatforms" Value="Win64;Mac;Linux;Android" />
+  <Property Name="Win64FoundationTags" Value="" />
+  <Property Name="MacFoundationTags" Value="" />
+  <Property Name="IOSFoundationTags" Value="" />
+  <Property Name="AndroidFoundationTags" Value="" />
+  <Property Name="LinuxFoundationTags" Value="" />
+  <Property Name="LinuxArm64FoundationTags" Value="" />
+  <Property Name="TVOSFoundationTags" Value="" />
+  <Property Name="Win64FoundationExtraCompilationArgs" Value="-allmodules" />
+  <Property Name="MacFoundationExtraCompilationArgs" Value="-allmodules" />
+  <Property Name="IOSFoundationExtraCompilationArgs" Value="-allmodules" />
+  <Property Name="AndroidFoundationExtraCompilationArgs" Value="-allmodules -architectures=arm64" />
+  <Property Name="LinuxFoundationExtraCompilationArgs" Value="-allmodules" />
+  <Property Name="LinuxArm64FoundationExtraCompilationArgs" Value="-allmodules" />
+  <Property Name="TVOSFoundationExtraCompilationArgs" Value="-allmodules" />
+  <Property Name="Win64FoundationRunSupported" Value="True" />
+  <Property Name="MacFoundationRunSupported" Value="True" />
+  <Property Name="IOSFoundationRunSupported" Value="False" />
+  <Property Name="AndroidFoundationRunSupported" Value="False" />
+  <Property Name="LinuxFoundationRunSupported" Value="True" />
+  <Property Name="LinuxArm64FoundationRunSupported" Value="True" />
+  <Property Name="TVOSFoundationRunSupported" Value="True" />
+  <Property Name="HoloLensFoundationTags" Value="" />
+  <Property Name="HoloLensFoundationExtraCompilationArgs" Value="-allmodules" />
+  <Property Name="HoloLensFoundationRunSupported" Value="True" />
+</BuildGraph>
\ No newline at end of file
diff --git a/Engine/Build/LowLevelTests/HeadlessChaos.xml b/Engine/Build/LowLevelTests/HeadlessChaos.xml
new file mode 100644
--- /dev/null
+++ b/Engine/Build/LowLevelTests/HeadlessChaos.xml
@@ -0,0 +1,39 @@
+﻿<?xml version="1.0" encoding="utf-8"?>
+<BuildGraph xmlns:xsi="http://www.w3.org/2001/XMLSchema-instance" xsi:schemaLocation="http://www.epicgames.com/BuildGraph ../../Build/Graph/Schema.xsd" xmlns="http://www.epicgames.com/BuildGraph">
+  <Property Name="TestNames" Value="$(TestNames);HeadlessChaos" />
+  <Property Name="HeadlessChaosUsesCatch2" Value="False" />
+  <Property Name="HeadlessChaosDisabled" Value="False" />
+  <Property Name="HeadlessChaosShort" Value="Headless Chaos" />
+  <Property Name="HeadlessChaosTarget" Value="HeadlessChaos" />
+  <Property Name="HeadlessChaosBinariesRelative" Value="Engine\Binaries" />
+  <Property Name="HeadlessChaosReportType" />
+  <Property Name="HeadlessChaosInitialExtraArgs" />
+  <Property Name="HeadlessChaosHasAfterSteps" Value="False" />
+  <Property Name="HeadlessChaosUsesCatch2" Value="False" />
+  <Option Name="RunHeadlessChaosTests" DefaultValue="False" Description="Run Headless Chaos Tests" />
+  <Property Name="HeadlessChaosSupportedPlatforms" Value="Win64" />
+  <Property Name="Win64HeadlessChaosTags" Value="" />
+  <Property Name="MacHeadlessChaosTags" Value="" />
+  <Property Name="IOSHeadlessChaosTags" Value="" />
+  <Property Name="AndroidHeadlessChaosTags" Value="" />
+  <Property Name="LinuxHeadlessChaosTags" Value="" />
+  <Property Name="LinuxArm64HeadlessChaosTags" Value="" />
+  <Property Name="TVOSHeadlessChaosTags" Value="" />
+  <Property Name="Win64HeadlessChaosExtraCompilationArgs" Value="" />
+  <Property Name="MacHeadlessChaosExtraCompilationArgs" Value="" />
+  <Property Name="IOSHeadlessChaosExtraCompilationArgs" Value="" />
+  <Property Name="AndroidHeadlessChaosExtraCompilationArgs" Value="" />
+  <Property Name="LinuxHeadlessChaosExtraCompilationArgs" Value="" />
+  <Property Name="LinuxArm64HeadlessChaosExtraCompilationArgs" Value="" />
+  <Property Name="TVOSHeadlessChaosExtraCompilationArgs" Value="" />
+  <Property Name="Win64HeadlessChaosRunSupported" Value="True" />
+  <Property Name="MacHeadlessChaosRunSupported" Value="True" />
+  <Property Name="IOSHeadlessChaosRunSupported" Value="False" />
+  <Property Name="AndroidHeadlessChaosRunSupported" Value="False" />
+  <Property Name="LinuxHeadlessChaosRunSupported" Value="True" />
+  <Property Name="LinuxArm64HeadlessChaosRunSupported" Value="True" />
+  <Property Name="TVOSHeadlessChaosRunSupported" Value="True" />
+  <Property Name="HoloLensHeadlessChaosTags" Value="" />
+  <Property Name="HoloLensHeadlessChaosExtraCompilationArgs" Value="" />
+  <Property Name="HoloLensHeadlessChaosRunSupported" Value="True" />
+</BuildGraph>
\ No newline at end of file
diff --git a/Engine/Build/LowLevelTests/NetworkPredictionPlugin.xml b/Engine/Build/LowLevelTests/NetworkPredictionPlugin.xml
new file mode 100644
--- /dev/null
+++ b/Engine/Build/LowLevelTests/NetworkPredictionPlugin.xml
@@ -0,0 +1,39 @@
+﻿<?xml version="1.0" encoding="utf-8"?>
+<BuildGraph xmlns:xsi="http://www.w3.org/2001/XMLSchema-instance" xsi:schemaLocation="http://www.epicgames.com/BuildGraph ../../Build/Graph/Schema.xsd" xmlns="http://www.epicgames.com/BuildGraph">
+  <Property Name="TestNames" Value="$(TestNames);NetworkPredictionPlugin" />
+  <Property Name="NetworkPredictionPluginUsesCatch2" Value="True" />
+  <Property Name="NetworkPredictionPluginDisabled" Value="False" />
+  <Property Name="NetworkPredictionPluginShort" Value="Net Prediction" />
+  <Property Name="NetworkPredictionPluginTarget" Value="NetworkPredictionTests" />
+  <Property Name="NetworkPredictionPluginBinariesRelative" Value="Engine\Binaries" />
+  <Property Name="NetworkPredictionPluginReportType" />
+  <Property Name="NetworkPredictionPluginInitialExtraArgs" />
+  <Property Name="NetworkPredictionPluginHasAfterSteps" Value="False" />
+  <Property Name="NetworkPredictionPluginUsesCatch2" Value="True" />
+  <Option Name="RunNetworkPredictionPluginTests" DefaultValue="False" Description="Run Net Prediction Tests" />
+  <Property Name="NetworkPredictionPluginSupportedPlatforms" Value="Win64;Linux" />
+  <Property Name="Win64NetworkPredictionPluginTags" Value="" />
+  <Property Name="MacNetworkPredictionPluginTags" Value="" />
+  <Property Name="IOSNetworkPredictionPluginTags" Value="" />
+  <Property Name="AndroidNetworkPredictionPluginTags" Value="" />
+  <Property Name="LinuxNetworkPredictionPluginTags" Value="" />
+  <Property Name="LinuxArm64NetworkPredictionPluginTags" Value="" />
+  <Property Name="TVOSNetworkPredictionPluginTags" Value="" />
+  <Property Name="Win64NetworkPredictionPluginExtraCompilationArgs" Value="" />
+  <Property Name="MacNetworkPredictionPluginExtraCompilationArgs" Value="" />
+  <Property Name="IOSNetworkPredictionPluginExtraCompilationArgs" Value="" />
+  <Property Name="AndroidNetworkPredictionPluginExtraCompilationArgs" Value="" />
+  <Property Name="LinuxNetworkPredictionPluginExtraCompilationArgs" Value="" />
+  <Property Name="LinuxArm64NetworkPredictionPluginExtraCompilationArgs" Value="" />
+  <Property Name="TVOSNetworkPredictionPluginExtraCompilationArgs" Value="" />
+  <Property Name="Win64NetworkPredictionPluginRunSupported" Value="True" />
+  <Property Name="MacNetworkPredictionPluginRunSupported" Value="True" />
+  <Property Name="IOSNetworkPredictionPluginRunSupported" Value="False" />
+  <Property Name="AndroidNetworkPredictionPluginRunSupported" Value="False" />
+  <Property Name="LinuxNetworkPredictionPluginRunSupported" Value="True" />
+  <Property Name="LinuxArm64NetworkPredictionPluginRunSupported" Value="True" />
+  <Property Name="TVOSNetworkPredictionPluginRunSupported" Value="True" />
+  <Property Name="HoloLensNetworkPredictionPluginTags" Value="" />
+  <Property Name="HoloLensNetworkPredictionPluginExtraCompilationArgs" Value="" />
+  <Property Name="HoloLensNetworkPredictionPluginRunSupported" Value="True" />
+</BuildGraph>
\ No newline at end of file
diff --git a/Engine/Build/LowLevelTests/ReplicationSystem.xml b/Engine/Build/LowLevelTests/ReplicationSystem.xml
new file mode 100644
--- /dev/null
+++ b/Engine/Build/LowLevelTests/ReplicationSystem.xml
@@ -0,0 +1,39 @@
+﻿<?xml version="1.0" encoding="utf-8"?>
+<BuildGraph xmlns:xsi="http://www.w3.org/2001/XMLSchema-instance" xsi:schemaLocation="http://www.epicgames.com/BuildGraph ../../Build/Graph/Schema.xsd" xmlns="http://www.epicgames.com/BuildGraph">
+  <Property Name="TestNames" Value="$(TestNames);ReplicationSystem" />
+  <Property Name="ReplicationSystemUsesCatch2" Value="True" />
+  <Property Name="ReplicationSystemDisabled" Value="False" />
+  <Property Name="ReplicationSystemShort" Value="Replication System" />
+  <Property Name="ReplicationSystemTarget" Value="ReplicationSystemLowLevelTests" />
+  <Property Name="ReplicationSystemBinariesRelative" Value="Engine\Binaries" />
+  <Property Name="ReplicationSystemReportType" />
+  <Property Name="ReplicationSystemInitialExtraArgs" />
+  <Property Name="ReplicationSystemHasAfterSteps" Value="False" />
+  <Property Name="ReplicationSystemUsesCatch2" Value="True" />
+  <Option Name="RunReplicationSystemTests" DefaultValue="False" Description="Run Replication System Tests" />
+  <Property Name="ReplicationSystemSupportedPlatforms" Value="Win64;Mac;Linux;Android" />
+  <Property Name="Win64ReplicationSystemTags" Value="" />
+  <Property Name="MacReplicationSystemTags" Value="" />
+  <Property Name="IOSReplicationSystemTags" Value="" />
+  <Property Name="AndroidReplicationSystemTags" Value="" />
+  <Property Name="LinuxReplicationSystemTags" Value="" />
+  <Property Name="LinuxArm64ReplicationSystemTags" Value="" />
+  <Property Name="TVOSReplicationSystemTags" Value="" />
+  <Property Name="Win64ReplicationSystemExtraCompilationArgs" Value="" />
+  <Property Name="MacReplicationSystemExtraCompilationArgs" Value="" />
+  <Property Name="IOSReplicationSystemExtraCompilationArgs" Value="" />
+  <Property Name="AndroidReplicationSystemExtraCompilationArgs" Value="" />
+  <Property Name="LinuxReplicationSystemExtraCompilationArgs" Value="" />
+  <Property Name="LinuxArm64ReplicationSystemExtraCompilationArgs" Value="" />
+  <Property Name="TVOSReplicationSystemExtraCompilationArgs" Value="" />
+  <Property Name="Win64ReplicationSystemRunSupported" Value="True" />
+  <Property Name="MacReplicationSystemRunSupported" Value="True" />
+  <Property Name="IOSReplicationSystemRunSupported" Value="False" />
+  <Property Name="AndroidReplicationSystemRunSupported" Value="False" />
+  <Property Name="LinuxReplicationSystemRunSupported" Value="True" />
+  <Property Name="LinuxArm64ReplicationSystemRunSupported" Value="True" />
+  <Property Name="TVOSReplicationSystemRunSupported" Value="True" />
+  <Property Name="HoloLensReplicationSystemTags" Value="" />
+  <Property Name="HoloLensReplicationSystemExtraCompilationArgs" Value="" />
+  <Property Name="HoloLensReplicationSystemRunSupported" Value="True" />
+</BuildGraph>
\ No newline at end of file
diff --git a/Engine/Platforms/Hololens/Build/LowLevelTests/Foundation.xml b/Engine/Platforms/Hololens/Build/LowLevelTests/Foundation.xml
new file mode 100644
--- /dev/null
+++ b/Engine/Platforms/Hololens/Build/LowLevelTests/Foundation.xml
@@ -0,0 +1,6 @@
+﻿<?xml version="1.0" encoding="utf-8"?>
+<BuildGraph xmlns:xsi="http://www.w3.org/2001/XMLSchema-instance" xsi:schemaLocation="http://www.epicgames.com/BuildGraph ../../Build/Graph/Schema.xsd" xmlns="http://www.epicgames.com/BuildGraph">
+  <Property Name="HoloLensFoundationTags" Value="" />
+  <Property Name="HoloLensFoundationExtraCompilationArgs" Value="-allmodules" />
+  <Property Name="HoloLensFoundationRunSupported" Value="True" />
+</BuildGraph>
\ No newline at end of file
diff --git a/Engine/Platforms/Hololens/Build/LowLevelTests/HeadlessChaos.xml b/Engine/Platforms/Hololens/Build/LowLevelTests/HeadlessChaos.xml
new file mode 100644
--- /dev/null
+++ b/Engine/Platforms/Hololens/Build/LowLevelTests/HeadlessChaos.xml
@@ -0,0 +1,6 @@
+﻿<?xml version="1.0" encoding="utf-8"?>
+<BuildGraph xmlns:xsi="http://www.w3.org/2001/XMLSchema-instance" xsi:schemaLocation="http://www.epicgames.com/BuildGraph ../../Build/Graph/Schema.xsd" xmlns="http://www.epicgames.com/BuildGraph">
+  <Property Name="HoloLensHeadlessChaosTags" Value="" />
+  <Property Name="HoloLensHeadlessChaosExtraCompilationArgs" Value="" />
+  <Property Name="HoloLensHeadlessChaosRunSupported" Value="True" />
+</BuildGraph>
\ No newline at end of file
diff --git a/Engine/Platforms/Hololens/Build/LowLevelTests/NetworkPredictionPlugin.xml b/Engine/Platforms/Hololens/Build/LowLevelTests/NetworkPredictionPlugin.xml
new file mode 100644
--- /dev/null
+++ b/Engine/Platforms/Hololens/Build/LowLevelTests/NetworkPredictionPlugin.xml
@@ -0,0 +1,6 @@
+﻿<?xml version="1.0" encoding="utf-8"?>
+<BuildGraph xmlns:xsi="http://www.w3.org/2001/XMLSchema-instance" xsi:schemaLocation="http://www.epicgames.com/BuildGraph ../../Build/Graph/Schema.xsd" xmlns="http://www.epicgames.com/BuildGraph">
+  <Property Name="HoloLensNetworkPredictionPluginTags" Value="" />
+  <Property Name="HoloLensNetworkPredictionPluginExtraCompilationArgs" Value="" />
+  <Property Name="HoloLensNetworkPredictionPluginRunSupported" Value="True" />
+</BuildGraph>
\ No newline at end of file
diff --git a/Engine/Platforms/Hololens/Build/LowLevelTests/ReplicationSystem.xml b/Engine/Platforms/Hololens/Build/LowLevelTests/ReplicationSystem.xml
new file mode 100644
--- /dev/null
+++ b/Engine/Platforms/Hololens/Build/LowLevelTests/ReplicationSystem.xml
@@ -0,0 +1,6 @@
+﻿<?xml version="1.0" encoding="utf-8"?>
+<BuildGraph xmlns:xsi="http://www.w3.org/2001/XMLSchema-instance" xsi:schemaLocation="http://www.epicgames.com/BuildGraph ../../Build/Graph/Schema.xsd" xmlns="http://www.epicgames.com/BuildGraph">
+  <Property Name="HoloLensReplicationSystemTags" Value="" />
+  <Property Name="HoloLensReplicationSystemExtraCompilationArgs" Value="" />
+  <Property Name="HoloLensReplicationSystemRunSupported" Value="True" />
+</BuildGraph>
\ No newline at end of file
