LLT Foundation Tests iOS - Compile only
#rnx

diff --git a/Engine/Build/LowLevelTests/Foundation.xml b/Engine/Build/LowLevelTests/Foundation.xml
--- a/Engine/Build/LowLevelTests/Foundation.xml
+++ b/Engine/Build/LowLevelTests/Foundation.xml
@@ -12,7 +12,7 @@
   <Property Name="FoundationHasAfterSteps" Value="False" />
   <Property Name="FoundationUsesCatch2" Value="True" />
   <Option Name="RunFoundationTests" DefaultValue="False" Description="Run Foundation Tests" />
-  <Property Name="FoundationSupportedPlatforms" Value="Win64;Mac;Linux;Android" />
+  <Property Name="FoundationSupportedPlatforms" Value="Win64;Mac;Linux;Android;iOS" />
   <Property Name="Win64FoundationTags" Value="" />
   <Property Name="MacFoundationTags" Value="" />
   <Property Name="IOSFoundationTags" Value="" />
