addressing duplicated generated properties by splitting the Web.xml

diff --git a/Engine/Build/LowLevelTests/Web.xml b/Engine/Build/LowLevelTests/Web.xml
new file mode 100644
--- /dev/null
+++ b/Engine/Build/LowLevelTests/Web.xml
@@ -0,0 +1,41 @@
+﻿<?xml version="1.0" encoding="utf-8"?>
+<BuildGraph xmlns:xsi="http://www.w3.org/2001/XMLSchema-instance" xsi:schemaLocation="http://www.epicgames.com/BuildGraph ../../Build/Graph/Schema.xsd" xmlns="http://www.epicgames.com/BuildGraph">
+  <Property Name="TestNames" Value="$(TestNames);Web" />
+  <Property Name="WebDisabled" Value="False" />
+  <Property Name="WebShort" Value="Web" />
+  <Property Name="WebStagesWithProjectFile" Value="False" />
+  <Property Name="WebTarget" Value="WebTests" />
+  <Property Name="WebBinariesRelative" Value="Engine\Binaries" />
+  <Property Name="WebReportType" Value="xml" />
+  <Property Name="WebGauntletArgs" Value="" />
+  <Property Name="WebExtraArgs" Value="" />
+  <Property Name="WebInitialExtraArgs" />
+  <Property Name="WebHasAfterSteps" Value="True" />
+  <Property Name="WebUsesCatch2" Value="True" />
+  <Option Name="RunWebTests" DefaultValue="False" Description="Run Web Tests" />
+  <Property Name="WebSupportedPlatforms" Value="Win64" />
+  <Property Name="Win64WebTags" Value="" />
+  <Property Name="MacWebTags" Value="" />
+  <Property Name="IOSWebTags" Value="" />
+  <Property Name="AndroidWebTags" Value="" />
+  <Property Name="LinuxWebTags" Value="" />
+  <Property Name="LinuxArm64WebTags" Value="" />
+  <Property Name="TVOSWebTags" Value="" />
+  <Property Name="Win64WebExtraCompilationArgs" Value="" />
+  <Property Name="MacWebExtraCompilationArgs" Value="" />
+  <Property Name="IOSWebExtraCompilationArgs" Value="" />
+  <Property Name="AndroidWebExtraCompilationArgs" Value="" />
+  <Property Name="LinuxWebExtraCompilationArgs" Value="" />
+  <Property Name="LinuxArm64WebExtraCompilationArgs" Value="" />
+  <Property Name="TVOSWebExtraCompilationArgs" Value="" />
+  <Property Name="Win64WebRunSupported" Value="True" />
+  <Property Name="MacWebRunSupported" Value="True" />
+  <Property Name="IOSWebRunSupported" Value="False" />
+  <Property Name="AndroidWebRunSupported" Value="False" />
+  <Property Name="LinuxWebRunSupported" Value="True" />
+  <Property Name="LinuxArm64WebRunSupported" Value="True" />
+  <Property Name="TVOSWebRunSupported" Value="True" />
+  <Property Name="VisionOSWebTags" Value="" />
+  <Property Name="VisionOSWebExtraCompilationArgs" Value="" />
+  <Property Name="VisionOSWebRunSupported" Value="True" />
+</BuildGraph>
\ No newline at end of file
