Cleanup docker cache for LinuxArm64 tests.

diff --git a/Engine/Build/LowLevelTests/LinuxArm64/docker_cleanup.bat b/Engine/Build/LowLevelTests/LinuxArm64/docker_cleanup.bat
--- a/Engine/Build/LowLevelTests/LinuxArm64/docker_cleanup.bat
+++ b/Engine/Build/LowLevelTests/LinuxArm64/docker_cleanup.bat
@@ -1,4 +1,5 @@
 @ECHO OFF
 set AppNameLowercase=%1
+docker system prune -a -f 2>NUL || exit 0
 docker rm -f %AppNameLowercase%-linuxarm64-container 2>NUL || exit 0
 docker rmi %AppNameLowercase%-linuxarm64-image 2>NUL || exit 0
\ No newline at end of file
diff --git a/Engine/Build/LowLevelTests/LinuxArm64/docker_cleanup.sh b/Engine/Build/LowLevelTests/LinuxArm64/docker_cleanup.sh
--- a/Engine/Build/LowLevelTests/LinuxArm64/docker_cleanup.sh
+++ b/Engine/Build/LowLevelTests/LinuxArm64/docker_cleanup.sh
@@ -1,5 +1,6 @@
 #!/bin/bash
 
 AppNameLowercase=$1
+docker system prune -a -f 2>/dev/null || exit 0
 docker rm -f ${AppNameLowercase}-linuxarm64-container 2>/dev/null || exit 0
 docker rmi ${AppNameLowercase}-linuxarm64-image 2>/dev/null || exit 0
\ No newline at end of file
