Enable web tests for iOS and Mac.

diff --git a/Engine/Build/LowLevelTests/Web.xml b/Engine/Build/LowLevelTests/Web.xml
--- a/Engine/Build/LowLevelTests/Web.xml
+++ b/Engine/Build/LowLevelTests/Web.xml
@@ -13,7 +13,7 @@
 	<Property Name="WebHasAfterSteps" Value="True" />
 	<Property Name="WebUsesCatch2" Value="True" />
 	<Option Name="RunWebTests" DefaultValue="False" Description="Run Web Tests" />
-	<Property Name="WebSupportedPlatforms" Value="Win64;Linux;Android" />
+	<Property Name="WebSupportedPlatforms" Value="Win64;Linux;Android;IOS;Mac" />
 	<Property Name="Win64WebTags" Value="" />
 	<Property Name="MacWebTags" Value="" />
 	<Property Name="IOSWebTags" Value="" />
@@ -21,23 +21,23 @@
 	<Property Name="LinuxWebTags" Value="" />
 	<Property Name="LinuxArm64WebTags" Value="" />
 	<Property Name="TVOSWebTags" Value="" />
-	<Property Name="Win64WebExtraCompilationArgs" Value="" />
-	<Property Name="MacWebExtraCompilationArgs" Value="" />
-	<Property Name="IOSWebExtraCompilationArgs" Value="" />
-	<Property Name="AndroidWebExtraCompilationArgs" Value="" />
-	<Property Name="LinuxWebExtraCompilationArgs" Value="" />
-	<Property Name="LinuxArm64WebExtraCompilationArgs" Value="" />
-	<Property Name="TVOSWebExtraCompilationArgs" Value="" />
+	<Property Name="Win64WebExtraCompilationArgs" Value="-allmodules" />
+	<Property Name="MacWebExtraCompilationArgs" Value="-allmodules" />
+	<Property Name="IOSWebExtraCompilationArgs" Value="-allmodules" />
+	<Property Name="AndroidWebExtraCompilationArgs" Value="-allmodules -architectures=arm64" />
+	<Property Name="LinuxWebExtraCompilationArgs" Value="-allmodules" />
+	<Property Name="LinuxArm64WebExtraCompilationArgs" Value="-allmodules" />
+	<Property Name="TVOSWebExtraCompilationArgs" Value="-allmodules" />
 	<Property Name="Win64WebRunSupported" Value="True" />
 	<Property Name="MacWebRunSupported" Value="True" />
 	<Property Name="IOSWebRunSupported" Value="False" />
 	<Property Name="AndroidWebRunSupported" Value="False" />
 	<Property Name="LinuxWebRunSupported" Value="True" />
 	<Property Name="LinuxArm64WebRunSupported" Value="True" />
-	<Property Name="TVOSWebRunSupported" Value="True" />
+	<Property Name="TVOSWebRunSupported" Value="False" />
 	<Property Name="VisionOSWebTags" Value="" />
-	<Property Name="VisionOSWebExtraCompilationArgs" Value="" />
-	<Property Name="VisionOSWebRunSupported" Value="True" />
+	<Property Name="VisionOSWebExtraCompilationArgs" Value="-allmodules" />
+	<Property Name="VisionOSWebRunSupported" Value="False" />
 	<Property Name="Win64WebRunContainerized" Value="False" />
 	<Property Name="MacWebRunContainerized" Value="False" />
 	<Property Name="IOSWebRunContainerized" Value="False" />
diff --git a/Engine/Source/Programs/WebTestsServer/createenv.sh b/Engine/Source/Programs/WebTestsServer/createenv.sh
--- a/Engine/Source/Programs/WebTestsServer/createenv.sh
+++ b/Engine/Source/Programs/WebTestsServer/createenv.sh
@@ -1,11 +1,14 @@
 if [ -d "env" ]; then
 	. ./env/bin/activate
 else
-	if [ "$OSTYPE" = "darwin" ]; then
-		../../../Binaries/ThirdParty/Python3/Mac/bin/python3 -m venv env
-	else
-		../../../Binaries/ThirdParty/Python3/Linux/bin/python3 -m venv env
-	fi
+	case "$OSTYPE" in
+		darwin*)
+			../../../Binaries/ThirdParty/Python3/Mac/bin/python3 -m venv env
+			;;
+		*)
+			../../../Binaries/ThirdParty/Python3/Linux/bin/python3 -m venv env
+			;;
+	esac
 	. ./env/bin/activate
 	python -m pip install --upgrade pip
 	pip install -r requirements.txt
