Skip more [Slow] tagged tests on other platforms for FoundationTests
#rnx

diff --git a/Engine/Build/LowLevelTests/Foundation.xml b/Engine/Build/LowLevelTests/Foundation.xml
--- a/Engine/Build/LowLevelTests/Foundation.xml
+++ b/Engine/Build/LowLevelTests/Foundation.xml
@@ -19,7 +19,7 @@
   <Property Name="MacFoundationTags" Value="" />
   <Property Name="IOSFoundationTags" Value="" />
   <Property Name="AndroidFoundationTags" Value="" />
-  <Property Name="LinuxFoundationTags" Value="" />
+  <Property Name="LinuxFoundationTags" Value="~[.]~[Slow]" />
   <Property Name="LinuxArm64FoundationTags" Value="" />
   <Property Name="TVOSFoundationTags" Value="" />
   <Property Name="Win64FoundationExtraCompilationArgs" Value="-allmodules" />
