Disable aftersteps in WebTests due to Horde.

diff --git a/Engine/Build/LowLevelTests/Web.xml b/Engine/Build/LowLevelTests/Web.xml
--- a/Engine/Build/LowLevelTests/Web.xml
+++ b/Engine/Build/LowLevelTests/Web.xml
@@ -3,10 +3,10 @@
   <Option Name="RunWebTests" DefaultValue="" Description="Run Web Tests" />
   <Property Name="TestNames" Value="$(TestNames);Web" />
   <Extend Name="RunAllTests">
-    <Expand Name="DeployAndTest" Platform="Win64" TestName="Web" ShortName="Web" TargetName="WebTests" BinaryRelativePath="Engine\Binaries" ReportType="xml" HasAfterSteps="True" ExtraCompilationArgs="-allmodules" />
-    <Expand Name="DeployAndTest" Platform="Linux" TestName="Web" ShortName="Web" TargetName="WebTests" BinaryRelativePath="Engine\Binaries" ReportType="xml" HasAfterSteps="True" ExtraCompilationArgs="-allmodules" />
-    <Expand Name="DeployAndTest" Platform="Mac" TestName="Web" ShortName="Web" TargetName="WebTests" BinaryRelativePath="Engine\Binaries" ReportType="xml" HasAfterSteps="True" ExtraCompilationArgs="-allmodules" />
-    <Expand Name="DeployAndTest" Platform="Android" TestName="Web" ShortName="Web" TargetName="WebTests" BinaryRelativePath="Engine\Binaries" ReportType="xml" HasAfterSteps="True" ExtraCompilationArgs="-allmodules  -architectures=arm64" RunUnsupported="True" />
-    <Expand Name="DeployAndTest" Platform="IOS" TestName="Web" ShortName="Web" TargetName="WebTests" BinaryRelativePath="Engine\Binaries" ReportType="xml" HasAfterSteps="True" ExtraCompilationArgs="-allmodules" RunUnsupported="True" />
+    <Expand Name="DeployAndTest" Platform="Win64" TestName="Web" ShortName="Web" TargetName="WebTests" BinaryRelativePath="Engine\Binaries" ReportType="xml" ExtraCompilationArgs="-allmodules" />
+    <Expand Name="DeployAndTest" Platform="Linux" TestName="Web" ShortName="Web" TargetName="WebTests" BinaryRelativePath="Engine\Binaries" ReportType="xml" ExtraCompilationArgs="-allmodules" />
+    <Expand Name="DeployAndTest" Platform="Mac" TestName="Web" ShortName="Web" TargetName="WebTests" BinaryRelativePath="Engine\Binaries" ReportType="xml" ExtraCompilationArgs="-allmodules" />
+    <Expand Name="DeployAndTest" Platform="Android" TestName="Web" ShortName="Web" TargetName="WebTests" BinaryRelativePath="Engine\Binaries" ReportType="xml" ExtraCompilationArgs="-allmodules  -architectures=arm64" RunUnsupported="True" />
+    <Expand Name="DeployAndTest" Platform="IOS" TestName="Web" ShortName="Web" TargetName="WebTests" BinaryRelativePath="Engine\Binaries" ReportType="xml" ExtraCompilationArgs="-allmodules" RunUnsupported="True" />
   </Extend>
 </BuildGraph>
\ No newline at end of file
diff --git a/Engine/Source/Programs/WebTests/WebTests.Build.cs b/Engine/Source/Programs/WebTests/WebTests.Build.cs
--- a/Engine/Source/Programs/WebTests/WebTests.Build.cs
+++ b/Engine/Source/Programs/WebTests/WebTests.Build.cs
@@ -17,7 +17,7 @@ public class WebTests : TestModuleRules
 			TestMetadata.SupportedPlatforms.Add(UnrealTargetPlatform.Android);
 			TestMetadata.SupportedPlatforms.Add(UnrealTargetPlatform.IOS);
 
-			TestMetadata.HasAfterSteps = true;
+			TestMetadata.HasAfterSteps = false;
 			TestMetadata.PlatformCompilationExtraArgs.Add(UnrealTargetPlatform.Win64, "-allmodules");
 			TestMetadata.PlatformCompilationExtraArgs.Add(UnrealTargetPlatform.Linux, "-allmodules");
 			TestMetadata.PlatformCompilationExtraArgs.Add(UnrealTargetPlatform.LinuxArm64, "-allmodules");
