Stop running PlainPropsTests on CI

#rnx

diff --git a/Engine/Build/LowLevelTests/PlainProps.xml b/Engine/Build/LowLevelTests/PlainProps.xml
--- a/Engine/Build/LowLevelTests/PlainProps.xml
+++ b/Engine/Build/LowLevelTests/PlainProps.xml
@@ -3,8 +3,6 @@
   <Option Name="RunPlainPropsTests" DefaultValue="" Description="Run PlainProps Tests" />
   <Property Name="TestNames" Value="$(TestNames);PlainProps" />
   <Extend Name="RunAllTests">
-    <Expand Name="DeployAndTest" Platform="Win64" TestName="PlainProps" ShortName="PlainProps" TargetName="PlainPropsTests" BinaryRelativePath="Engine\Binaries" ReportType="xml" ExtraCompilationArgs="-allmodules" />
-    <Expand Name="DeployAndTest" Platform="Linux" TestName="PlainProps" ShortName="PlainProps" TargetName="PlainPropsTests" BinaryRelativePath="Engine\Binaries" ReportType="xml" Tags="~[.]~[Slow]" ExtraCompilationArgs="-allmodules" />
-    <Expand Name="DeployAndTest" Platform="Mac" TestName="PlainProps" ShortName="PlainProps" TargetName="PlainPropsTests" BinaryRelativePath="Engine\Binaries" ReportType="xml" ExtraCompilationArgs="-allmodules" />
+    <Expand Name="DeployAndTest" Platform="Win64" TestName="PlainProps" ShortName="PlainProps" TargetName="PlainPropsTests" BinaryRelativePath="Engine\Binaries" ReportType="xml" ExtraCompilationArgs="-allmodules" Deactivated="True" />
   </Extend>
 </BuildGraph>
\ No newline at end of file
diff --git a/Engine/Source/Programs/PlainPropsTests/PlainPropsTests.Build.cs b/Engine/Source/Programs/PlainPropsTests/PlainPropsTests.Build.cs
--- a/Engine/Source/Programs/PlainPropsTests/PlainPropsTests.Build.cs
+++ b/Engine/Source/Programs/PlainPropsTests/PlainPropsTests.Build.cs
@@ -11,8 +11,9 @@ public class PlainPropsTests : TestModuleRules
 			TestMetadata.TestName = "PlainProps";
 			TestMetadata.TestShortName = "PlainProps";
 			TestMetadata.ReportType = "xml";
-			TestMetadata.SupportedPlatforms.Add(UnrealTargetPlatform.Linux);
-			TestMetadata.SupportedPlatforms.Add(UnrealTargetPlatform.Mac);
+			TestMetadata.Deactivated = true;
+			//TestMetadata.SupportedPlatforms.Add(UnrealTargetPlatform.Linux);
+			//TestMetadata.SupportedPlatforms.Add(UnrealTargetPlatform.Mac);
 
 			string PlatformCompilationArgs;
 			foreach (var Platform in UnrealTargetPlatform.GetValidPlatforms())
@@ -22,7 +23,7 @@ public class PlainPropsTests : TestModuleRules
 			}
 
 			// Platform-specific tags
-			TestMetadata.PlatformTags.Add(UnrealTargetPlatform.Linux, "~[.]~[Slow]");
+			//TestMetadata.PlatformTags.Add(UnrealTargetPlatform.Linux, "~[.]~[Slow]");
 		}
 	}
 	public PlainPropsTests(ReadOnlyTargetRules Target) : base(Target, true)
