OnlineServicesNullTests: Add XBox, Nintendo, IOS, Android platform support.

#Synced-CL 48810755

diff --git a/Engine/Build/LowLevelTests/OnlineServicesNull.xml b/Engine/Build/LowLevelTests/OnlineServicesNull.xml
--- a/Engine/Build/LowLevelTests/OnlineServicesNull.xml
+++ b/Engine/Build/LowLevelTests/OnlineServicesNull.xml
@@ -8,7 +8,7 @@
   <Extend Name="RunAllTests">
     <Expand Name="DeployAndTest" Platform="Win64" TestName="OnlineServicesNull" ShortName="OnlineServicesNull" StagesWithProjectFile="True" TargetName="OnlineServicesNullTests" BinaryRelativePath="Engine\Binaries" ReportType="xml" ExtraArgs="-dumpiniloads" Tags="~[StringProcessing]" />
     <Expand Name="DeployAndTest" Platform="Linux" TestName="OnlineServicesNull" ShortName="OnlineServicesNull" StagesWithProjectFile="True" TargetName="OnlineServicesNullTests" BinaryRelativePath="Engine\Binaries" ReportType="xml" ExtraArgs="-dumpiniloads" Tags="~[StringProcessing]" />
-    <Expand Name="DeployAndTest" Platform="Mac" TestName="OnlineServicesNull" ShortName="OnlineServicesNull" StagesWithProjectFile="True" TargetName="OnlineServicesNullTests" BinaryRelativePath="Engine\Binaries" ReportType="xml" ExtraArgs="-dumpiniloads" Tags="~[StringProcessing]" />
+    <!-- <Expand Name="DeployAndTest" Platform="Mac" TestName="OnlineServicesNull" ShortName="OnlineServicesNull" StagesWithProjectFile="True" TargetName="OnlineServicesNullTests" BinaryRelativePath="Engine\Binaries" ReportType="xml" ExtraArgs="-dumpiniloads" Tags="~[StringProcessing]" ExtraCompilationArgs="-ProgramArchitecture=x64" /> -->
     <!-- Currently needs test framework support, leave out from running on Horde for now. -->
     <!-- <Expand Name="DeployAndTest" Platform="IOS" TestName="OnlineServicesNull" ShortName="OnlineServicesNull" StagesWithProjectFile="True" TargetName="OnlineServicesNullTests" BinaryRelativePath="Engine\Binaries" ReportType="xml" PlatformGauntletArgs="~[StringProcessing]" ExtraArgs="-dumpiniloads" Tags="~[StringProcessing]" RunUnsupported="True" /> -->
     <!-- <Expand Name="DeployAndTest" Platform="Android" TestName="OnlineServicesNull" ShortName="OnlineServicesNull" StagesWithProjectFile="True" TargetName="OnlineServicesNullTests" BinaryRelativePath="Engine\Binaries" ReportType="xml" PlatformGauntletArgs="~[StringProcessing]" ExtraArgs="-dumpiniloads" Tags="~[StringProcessing]" RunUnsupported="True" /> -->
